Add ADM plugin to engine

diff --git a/Engine/Plugins/Runtime/ADM/ADM.uplugin b/Engine/Plugins/Runtime/ADM/ADM.uplugin
new file mode 100644
--- /dev/null
+++ b/Engine/Plugins/Runtime/ADM/ADM.uplugin
@@ -0,0 +1,33 @@
+{
+	"FileVersion": 3,
+	"Version": 1,
+	"VersionName": "1.0",
+	"FriendlyName": "Audio Definition Model (ADM)",
+	"Description": "Currently only supports output spatialized using WASAPI aggregate output channels and spatial ADM information transmitted using Open Sound Control (OSC).",
+	"Category": "Audio",
+	"CreatedBy": "Epic Games, Inc.",
+	"CreatedByURL": "",
+	"DocsURL": "",
+	"MarketplaceURL": "",
+	"SupportURL": "",
+	"CanContainContent": true,
+	"IsBetaVersion": false,
+	"IsExperimentalVersion": true,
+	"Installed": false,
+	"Modules": [
+		{
+			"Name": "ADMSpatialization",
+			"Type": "Runtime",
+			"LoadingPhase": "Default",
+			"PlatformAllowList": [
+				"Win64"
+			]
+		}
+	],
+	"Plugins": [
+		{
+			"Name": "OSC",
+			"Enabled": true
+		}
+	]
+}
\ No newline at end of file
diff --git a/Engine/Plugins/Runtime/ADM/Source/ADMSpatialization/ADMSpatialization.Build.cs b/Engine/Plugins/Runtime/ADM/Source/ADMSpatialization/ADMSpatialization.Build.cs
new file mode 100644
--- /dev/null
+++ b/Engine/Plugins/Runtime/ADM/Source/ADMSpatialization/ADMSpatialization.Build.cs
@@ -0,0 +1,32 @@
+// Copyright Epic Games, Inc. All Rights Reserved.
+
+namespace UnrealBuildTool.Rules
+{
+	public class ADMSpatialization : ModuleRules
+	{
+		public ADMSpatialization(ReadOnlyTargetRules Target) : base(Target)
+		{
+			PublicDependencyModuleNames.AddRange(
+				new string[]
+				{
+					"AudioExtensions",
+					"Core",
+					"CoreUObject",
+					"DeveloperSettings",
+					"Engine",
+					"Networking",
+					"OSC"
+				}
+			);
+
+			PrivateDependencyModuleNames.AddRange(
+				new string[]
+				{
+					"AudioMixerCore",
+					"AudioMixer",
+					"SignalProcessing"
+				}
+			);
+		}
+	}
+}
diff --git a/Engine/Plugins/Runtime/ADM/Source/ADMSpatialization/Private/ADMDirectOutChannel.cpp b/Engine/Plugins/Runtime/ADM/Source/ADMSpatialization/Private/ADMDirectOutChannel.cpp
new file mode 100644
--- /dev/null
+++ b/Engine/Plugins/Runtime/ADM/Source/ADMSpatialization/Private/ADMDirectOutChannel.cpp
@@ -0,0 +1,80 @@
+// Copyright Epic Games, Inc. All Rights Reserved.
+
+#include "ADMDirectOutChannel.h"
+
+#include "Audio.h"
+#include "AudioMixer.h"
+#include "DSP/FloatArrayMath.h"
+#include "Logging/LogMacros.h"
+
+
+namespace UE::ADM::Spatialization
+{
+	TUniquePtr<Audio::FAlignedFloatBuffer> FSourceDirectOut::SilienceBuffer;
+
+	FSourceDirectOut::FSourceDirectOut(const int32 InChannelIndex, const uint32 InNumSamples, Audio::IAudioMixerPlatformInterface* InMixerPlatform)
+	{
+		ChannelIndex = InChannelIndex;
+		MixerPlatform = InMixerPlatform;
+		InitializeTempBuffers(InNumSamples);
+	}
+
+	FSourceDirectOut::~FSourceDirectOut()
+	{
+		FScopeLock ScopeLock(&DestructorCriticalSection);
+
+		MixerPlatform = nullptr;
+	}
+
+	void FSourceDirectOut::InitializeTempBuffers(const uint32 InNumSamples)
+	{
+		if (!SilienceBuffer.IsValid())
+		{
+			SilienceBuffer = MakeUnique<Audio::FAlignedFloatBuffer>();
+
+			if (SilienceBuffer.IsValid())
+			{
+				SilienceBuffer->SetNumZeroed(InNumSamples);
+			}
+		}
+	}
+
+	void FSourceDirectOut::ProcessDirectOut(const FAudioPluginSourceInputData& InInputData)
+	{
+		if (!DestructorCriticalSection.TryLock())
+		{
+			return;
+		}
+
+		if (MixerPlatform)
+		{
+			if (InInputData.NumChannels == 1)
+			{
+				// Output is typically clamped by the mixer after all the sources have been mixed into the output buffer.
+				// Since we are sending each source to a direct output, clamp here.
+				Audio::FAlignedFloatBuffer* InputBuffer = const_cast<Audio::FAlignedFloatBuffer*>(InInputData.AudioBuffer);
+				Audio::ArrayRangeClamp(*InputBuffer, -1.0f, 1.0f);
+
+				MixerPlatform->SubmitDirectOutBuffer(ChannelIndex, *InInputData.AudioBuffer);
+			}
+		}
+
+		DestructorCriticalSection.Unlock();
+	}
+
+	void FSourceDirectOut::ProcessSilience()
+	{
+		if (!DestructorCriticalSection.TryLock())
+		{
+			return;
+		}
+
+		if (MixerPlatform)
+		{
+			MixerPlatform->SubmitDirectOutBuffer(ChannelIndex, *SilienceBuffer);
+		}
+
+		DestructorCriticalSection.Unlock();
+	}
+
+} // namespace UE::ADM::Spatialization
diff --git a/Engine/Plugins/Runtime/ADM/Source/ADMSpatialization/Private/ADMDirectOutChannel.h b/Engine/Plugins/Runtime/ADM/Source/ADMSpatialization/Private/ADMDirectOutChannel.h
new file mode 100644
--- /dev/null
+++ b/Engine/Plugins/Runtime/ADM/Source/ADMSpatialization/Private/ADMDirectOutChannel.h
@@ -0,0 +1,38 @@
+// Copyright Epic Games, Inc. All Rights Reserved.
+
+#pragma once
+
+#include "AudioMixer.h"
+#include "Containers/Array.h"
+#include "DSP/AlignedBuffer.h"
+#include "IAudioExtensionPlugin.h"
+
+
+namespace UE::ADM::Spatialization
+{
+	class ADMSPATIALIZATION_API FSourceDirectOut
+	{
+	public:
+		FSourceDirectOut() = delete;
+		FSourceDirectOut(const int32 InChannelIndex, const uint32 InNumSamples, Audio::IAudioMixerPlatformInterface* InMixerPlatform);
+		~FSourceDirectOut();
+
+		void ProcessDirectOut(const FAudioPluginSourceInputData& InInputData);
+		void ProcessSilience();
+
+		void SetIsActive(const bool bInIsActive) { bIsActive = bInIsActive;  }
+		bool GetIsActive() const { return bIsActive; }
+
+	private:
+		static TUniquePtr<Audio::FAlignedFloatBuffer> SilienceBuffer;
+
+		FCriticalSection DestructorCriticalSection;
+
+		bool bIsActive = false;
+		Audio::IAudioMixerPlatformInterface* MixerPlatform = nullptr;
+		int32 ChannelIndex = INDEX_NONE;
+
+		void InitializeTempBuffers(const uint32 InNumSamples);
+	};
+
+} // namespace UE::ADM::Spatialization
diff --git a/Engine/Plugins/Runtime/ADM/Source/ADMSpatialization/Private/ADMSpatialization.cpp b/Engine/Plugins/Runtime/ADM/Source/ADMSpatialization/Private/ADMSpatialization.cpp
new file mode 100644
--- /dev/null
+++ b/Engine/Plugins/Runtime/ADM/Source/ADMSpatialization/Private/ADMSpatialization.cpp
@@ -0,0 +1,245 @@
+// Copyright Epic Games, Inc. All Rights Reserved.
+
+#include "ADMSpatialization.h"
+
+#include "ADMDirectOutChannel.h"
+#include "ADMSpatializationLog.h"
+#include "ADMSpatializationModule.h"
+#include "AudioMixerDevice.h"
+#include "OSCClient.h"
+#include "OSCMessage.h"
+#include "OSCTypes.h"
+
+
+void UADMEngineSubsystem::SetSendIPAddress(const FString& InIPAddress, int32 Port)
+{
+	using namespace UE::ADM::Spatialization;
+
+	const FString EndpointStr = FString::Printf(TEXT("%s:%i"), *InIPAddress, Port);
+	FIPv4Endpoint Endpoint;
+	if (FIPv4Endpoint::Parse(EndpointStr, Endpoint))
+	{
+		FModule& Module = FModuleManager::Get().LoadModuleChecked<FModule>("ADMSpatialization");
+		FADMSpatializationFactory& SpatFactory = Module.GetFactory();
+		SpatFactory.SetSendIPEndpoint(Endpoint);
+	}
+	else
+	{
+		UE_LOG(LogADMSpatialization, Error, TEXT("Failed to parse specified ADM Spatialization client endpoint '%s'. Client IP not updated for ADM Spatialization."), *EndpointStr);
+	}
+}
+
+namespace UE::ADM::Spatialization
+{
+	namespace SpatPrivate
+	{
+		FString SendEndpointCVar;
+		FAutoConsoleVariableRef CVarOSCADMSendEndpoint(
+			TEXT("au.ADM.Spatialization.OSCSendEndpoint"),
+			SendEndpointCVar,
+			TEXT("Override send (client) endpoint for ADM spatializer messaging (in the form 127.0.0.1:8000).")
+			TEXT("Default: Empty (Does not override project setting)"),
+			ECVF_Default);
+
+		int32 PositionAddressOffsetCVar = 0;
+		FAutoConsoleVariableRef CVarOSCADMPositionAddressOffset(
+			TEXT("au.ADM.Spatialization.OSCPositionAddressOffset"),
+			SendEndpointCVar,
+			TEXT("Applies an index offset to all source Ids translated to OSC position source addresses.")
+			TEXT("Default: 0 (No offset)"),
+			ECVF_Default);
+
+		static const FString SystemName(TEXT("ADM Spatialization"));
+
+	} // namespace SpatPrivate
+
+	FADMClient::FADMClient(const FIPv4Endpoint& InEndpoint)
+	{
+		FIPv4Endpoint IPEndpoint = InEndpoint;
+		if (!SpatPrivate::SendEndpointCVar.IsEmpty())
+		{
+			FIPv4Endpoint::Parse(SpatPrivate::SendEndpointCVar, IPEndpoint);
+		}
+
+		ClientProxy = UE::OSC::IClientProxy::Create(SpatPrivate::SystemName);
+		ClientProxy->SetSendIPEndpoint(IPEndpoint);
+	}
+
+	FOSCAddress FADMClient::CreateConfigAddress(int32 SourceId, FString InMethod)
+	{
+		TArray<FString> Containers
+		{
+			TEXT("adm"),
+			TEXT("obj"),
+			FString::FromInt(SourceId),
+			TEXT("config")
+		};
+
+		FOSCAddress OSCAddress;
+		OSCAddress.Set(MoveTemp(Containers), MoveTemp(InMethod));
+		return OSCAddress;
+	}
+
+	FOSCAddress FADMClient::CreatePositionAddress(int32 SourceId, FString InMethod)
+	{
+		TArray<FString> Containers
+		{
+			TEXT("adm"),
+			TEXT("obj"),
+			FString::FromInt(SourceId + SpatPrivate::PositionAddressOffsetCVar)
+		};
+
+		FOSCAddress OSCAddress;
+		OSCAddress.Set(MoveTemp(Containers), MoveTemp(InMethod));
+		return OSCAddress;
+	}
+
+	void FADMClient::InitSource(int32 InSourceId, bool bCartesian)
+	{
+		using namespace UE::OSC;
+
+		if (ClientProxy.IsValid())
+		{
+			FOSCAddress OSCAddress = CreateConfigAddress(InSourceId, TEXT("cartesian"));
+			FOSCMessage Msg(MoveTemp(OSCAddress), { FOSCData((int32)(bCartesian)) });
+			ClientProxy->SendMessage(Msg);
+		}
+	}
+
+	void FADMClient::SetPosition(int32 InSourceId, const FVector& InPosition)
+	{
+		using namespace UE::OSC;
+
+		if (ClientProxy.IsValid())
+		{
+			FOSCAddress OSCAddress = CreatePositionAddress(InSourceId, TEXT("xyz"));
+			FOSCMessage Msg(MoveTemp(OSCAddress), { FOSCData((float)InPosition.X) , FOSCData((float)InPosition.Y), FOSCData((float)InPosition.Z) });
+			ClientProxy->SendMessage(Msg);
+		}
+	}
+
+	bool FADMClient::IsSet() const
+	{
+		return ClientProxy.IsValid();
+	}
+
+	FString FADMSpatializationFactory::GetDisplayName()
+	{
+		return SpatPrivate::SystemName;
+	}
+
+	bool FADMSpatializationFactory::SupportsPlatform(const FString& PlatformName)
+	{
+		return true;
+	}
+
+	TAudioSpatializationPtr FADMSpatializationFactory::CreateNewSpatializationPlugin(FAudioDevice* OwningDevice)
+	{
+		return MakeShared<FADMSpatialization>();
+	}
+
+	UClass* FADMSpatializationFactory::GetCustomSpatializationSettingsClass() const
+	{
+		return UADMSpatializationSourceSettings::StaticClass();
+	};
+
+	bool FADMSpatializationFactory::IsExternalSend()
+	{
+		return false;
+	}
+
+	int32 FADMSpatializationFactory::GetMaxSupportedChannels()
+	{
+		return 1;
+	}
+
+	void FADMSpatialization::Initialize(const FAudioPluginInitializationParams InitializationParams)
+	{
+		FModule& Module = FModuleManager::Get().LoadModuleChecked<FModule>("ADMSpatialization");
+		FADMSpatializationFactory& SpatFactory = Module.GetFactory();
+		const FIPv4Endpoint Endpoint = SpatFactory.GetSendIPEndpoint();
+
+		SetClient(FADMClient(Endpoint));
+
+		int32 NumDirectOutChannels = 0;
+		Audio::FMixerDevice* AudioMixerDevice = static_cast<Audio::FMixerDevice*>(InitializationParams.AudioDevicePtr);
+		if (!ensureMsgf(AudioMixerDevice, TEXT("Unable to initialize FADMSpatialization, null audio device")))
+		{
+			return;
+		}
+
+		Audio::IAudioMixerPlatformInterface* MixerPlatform = AudioMixerDevice->GetAudioMixerPlatform();
+		if (!ensureMsgf(MixerPlatform, TEXT("Unable to initialize FADMSpatialization, null mixer platform")))
+		{
+			return;
+		}
+
+		NumBedChannels = AudioMixerDevice->GetNumDeviceChannels();
+		NumDirectOutChannels = AudioMixerDevice->GetNumDirectOutChannels();
+		
+		NumSources = InitializationParams.NumSources;
+		SampleRate = InitializationParams.SampleRate;
+
+		checkf(NumSources >= 0 && NumSources <= AudioMixerDevice->GetMaxSources(), TEXT("NumSources is expected to be in the range of 0 to max sources for the mixer device"));
+		checkf(SampleRate == AudioMixerDevice->GetSampleRate(), TEXT("SampleRate is expected to match the mixer device sample rate"));
+
+		DirectOuts.Reset(InitializationParams.NumSources);
+		for (uint32 Index = 0; Index < InitializationParams.NumSources; Index++)
+		{
+			DirectOuts.Emplace(Index, InitializationParams.BufferLength, MixerPlatform);
+		}
+	}
+
+	void FADMSpatialization::Shutdown()
+	{
+		DirectOuts.Reset();
+	}
+
+	void FADMSpatialization::SetClient(FADMClient&& InClient)
+	{
+		Client = MoveTemp(InClient);
+	}
+
+	bool FADMSpatialization::IsSpatializationEffectInitialized() const
+	{
+		return true;
+	}
+
+	void FADMSpatialization::OnInitSource(const uint32 SourceId, const FName& AudioComponentUserId, const uint32 NumChannels, USpatializationPluginSourceSettingsBase* InSettings)
+	{
+		checkf((int32)SourceId < NumSources, TEXT("SourceId is expected to be less than the initialized max number of sources"));
+
+		// Use SourceId as channel index
+		DirectOuts[SourceId].SetIsActive(true);
+
+		Client.InitSource((int32)(SourceId));
+	}
+
+	void FADMSpatialization::OnReleaseSource(const uint32 SourceId)
+	{
+		checkf((int32)SourceId < NumSources, TEXT("SourceId is expected to be less than the initialized max number of sources"));
+
+		DirectOuts[SourceId].SetIsActive(false);
+	}
+
+	void FADMSpatialization::ProcessAudio(const FAudioPluginSourceInputData& InputData, FAudioPluginSourceOutputData& OutputData)
+	{
+		checkf(InputData.SourceId < NumSources, TEXT("SourceId is expected to be less than the initialized max number of sources"));
+
+		DirectOuts[InputData.SourceId].ProcessDirectOut(InputData);
+		Client.SetPosition(InputData.SourceId, InputData.SpatializationParams->EmitterPosition);
+	}
+
+	void FADMSpatialization::OnAllSourcesProcessed()
+	{
+		for (FSourceDirectOut& DirectOut : DirectOuts)
+		{
+			if (!DirectOut.GetIsActive())
+			{
+				// Send silience out non-active outputs to retain time alignment with other outputs
+				DirectOut.ProcessSilience();
+			}
+		}
+	}
+
+} // namespace UE::ADMSpatialization
diff --git a/Engine/Plugins/Runtime/ADM/Source/ADMSpatialization/Private/ADMSpatialization.h b/Engine/Plugins/Runtime/ADM/Source/ADMSpatialization/Private/ADMSpatialization.h
new file mode 100644
--- /dev/null
+++ b/Engine/Plugins/Runtime/ADM/Source/ADMSpatialization/Private/ADMSpatialization.h
@@ -0,0 +1,125 @@
+// Copyright Epic Games, Inc. All Rights Reserved.
+
+#pragma once
+
+#include "ADMDirectOutChannel.h"
+#include "AudioMixer.h"
+#include "DSP/AlignedBuffer.h"
+#include "HAL/CriticalSection.h"
+#include "IAudioExtensionPlugin.h"
+#include "OSCClient.h"
+#include "Subsystems/AudioEngineSubsystem.h"
+#include "Templates/UniquePtr.h"
+
+#include "ADMSpatialization.generated.h"
+
+
+// Forward Declarations
+class IOSCClientProxy;
+
+
+UCLASS()
+class ADMSPATIALIZATION_API UADMSpatializationSourceSettings : public USpatializationPluginSourceSettingsBase
+{
+	GENERATED_BODY()
+
+public:
+	UADMSpatializationSourceSettings() = default;
+};
+
+UCLASS()
+class ADMSPATIALIZATION_API UADMEngineSubsystem : public UAudioEngineSubsystem
+{
+	GENERATED_BODY()
+
+	//~ Begin USubsystem interface
+	virtual bool ShouldCreateSubsystem(UObject* Outer) const override { return true; }
+	virtual void Initialize(FSubsystemCollectionBase& Collection) override { }
+	virtual void Deinitialize() override { }
+	//~ End USubsystem interface
+
+	// Set the IP Address to communicate ADM updates over OSC to
+	UFUNCTION(BlueprintCallable, Category = "ADM|OSC", meta = (DisplayName = "Set Send IP Address", Keywords = "osc adm message"))
+	void SetSendIPAddress(const FString& IPAddress, int32 Port);
+};
+
+
+namespace UE::ADM::Spatialization
+{
+	class ADMSPATIALIZATION_API FADMSpatializationFactory : public IAudioSpatializationFactory
+	{
+	public:
+		FADMSpatializationFactory() = default;
+		~FADMSpatializationFactory() = default;
+
+		virtual FString GetDisplayName() override;
+		virtual bool SupportsPlatform(const FString& PlatformName) override;
+		virtual TAudioSpatializationPtr CreateNewSpatializationPlugin(FAudioDevice* OwningDevice) override;
+		virtual UClass* GetCustomSpatializationSettingsClass() const override;
+		virtual bool IsExternalSend() override;
+		virtual int32 GetMaxSupportedChannels() override;
+
+		void SetSendIPEndpoint(const FIPv4Endpoint& InIPEndpoint)
+		{
+			FScopeLock Lock(&ClientIPCritSec);
+			IPEndpoint = InIPEndpoint;
+		}
+
+		const FIPv4Endpoint& GetSendIPEndpoint() const
+		{
+			FScopeLock Lock(&ClientIPCritSec);
+			return IPEndpoint;
+		}
+
+	private:
+		mutable FCriticalSection ClientIPCritSec;
+
+		FIPv4Endpoint IPEndpoint;
+	};
+
+
+	class ADMSPATIALIZATION_API FADMClient
+	{
+	public:
+		FADMClient() = default;
+		FADMClient(const FIPv4Endpoint& InEndpoint);
+
+		static FOSCAddress CreateConfigAddress(int32 SourceId, FString InMethod);
+		static FOSCAddress CreatePositionAddress(int32 SourceId, FString InMethod);
+
+		void InitSource(int32 InSourceId, bool bCartesian = true);
+		bool IsSet() const;
+
+		void SetPosition(int32 InSourceId, const FVector& InPosition);
+
+	private:
+		TUniquePtr<UE::OSC::IClientProxy> ClientProxy;
+	};
+
+
+	class ADMSPATIALIZATION_API FADMSpatialization : public IAudioSpatialization
+	{
+	public:
+		FADMSpatialization() = default;
+		~FADMSpatialization() = default;
+
+		virtual void Initialize(const FAudioPluginInitializationParams InitializationParams) override;
+		virtual bool IsSpatializationEffectInitialized() const override;
+		virtual void OnInitSource(const uint32 SourceId, const FName& AudioComponentUserId, const uint32 NumChannels, USpatializationPluginSourceSettingsBase* InSettings) override;
+		virtual void OnReleaseSource(const uint32 SourceId) override;
+		virtual void ProcessAudio(const FAudioPluginSourceInputData& InputData, FAudioPluginSourceOutputData& OutputData) override;
+		virtual void OnAllSourcesProcessed() override;
+		virtual void Shutdown () override;
+
+		void SetClient(FADMClient&& InClient);
+
+	private:
+		int32 NumBedChannels = 0;
+		int32 NumSources = 0;
+		float SampleRate = 0.f;
+
+		TArray<FSourceDirectOut> DirectOuts;
+
+		FADMClient Client;
+	};
+}  // namespace UE::ADM::Spatialization
diff --git a/Engine/Plugins/Runtime/ADM/Source/ADMSpatialization/Private/ADMSpatializationLog.h b/Engine/Plugins/Runtime/ADM/Source/ADMSpatialization/Private/ADMSpatializationLog.h
new file mode 100644
--- /dev/null
+++ b/Engine/Plugins/Runtime/ADM/Source/ADMSpatialization/Private/ADMSpatializationLog.h
@@ -0,0 +1,6 @@
+// Copyright Epic Games, Inc. All Rights Reserved.
+#pragma once
+
+#include "Logging/LogMacros.h"
+
+DECLARE_LOG_CATEGORY_EXTERN(LogADMSpatialization, Log, All);
diff --git a/Engine/Plugins/Runtime/ADM/Source/ADMSpatialization/Private/ADMSpatializationModule.cpp b/Engine/Plugins/Runtime/ADM/Source/ADMSpatialization/Private/ADMSpatializationModule.cpp
new file mode 100644
--- /dev/null
+++ b/Engine/Plugins/Runtime/ADM/Source/ADMSpatialization/Private/ADMSpatializationModule.cpp
@@ -0,0 +1,46 @@
+// Copyright Epic Games, Inc. All Rights Reserved.
+
+#include "ADMSpatializationModule.h"
+
+#include "ADMSpatializationSettings.h"
+#include "ADMSpatializationLog.h"
+
+DEFINE_LOG_CATEGORY(LogADMSpatialization);
+
+
+namespace UE::ADM::Spatialization
+{
+	void FModule::StartupModule()
+	{
+		IModularFeatures::Get().RegisterModularFeature(IAudioSpatializationFactory::GetModularFeatureName(), &SpatializationFactory);
+
+		if (const UADMSpatializationSettings* Settings = GetDefault<UADMSpatializationSettings>())
+		{
+			FIPv4Address IPAddress;
+			if (FIPv4Address::Parse(Settings->IPAddress, IPAddress))
+			{
+				SpatializationFactory.SetSendIPEndpoint(FIPv4Endpoint(IPAddress, Settings->IPPort));
+			}
+			else
+			{
+				UE_LOG(LogADMSpatialization, Error, TEXT("Failed to parse specified default ADM Spatialization client endpoint in developer settings. Default IP not set for ADM Spatialization."));
+			}
+		}
+		else
+		{
+			UE_LOG(LogADMSpatialization, Error, TEXT("Failed to find ADM Spatialization Developer Settings. Default IP not set for ADM Spatialization."));
+		}
+	}
+
+	void FModule::ShutdownModule()
+	{
+		IModularFeatures::Get().UnregisterModularFeature(IAudioSpatializationFactory::GetModularFeatureName(), &SpatializationFactory);
+	};
+
+	FADMSpatializationFactory& FModule::GetFactory()
+	{
+		return SpatializationFactory;
+	}
+} // namespace UE::ADM::Spatialization
+
+IMPLEMENT_MODULE(UE::ADM::Spatialization::FModule, ADMSpatialization)
diff --git a/Engine/Plugins/Runtime/ADM/Source/ADMSpatialization/Private/ADMSpatializationModule.h b/Engine/Plugins/Runtime/ADM/Source/ADMSpatialization/Private/ADMSpatializationModule.h
new file mode 100644
--- /dev/null
+++ b/Engine/Plugins/Runtime/ADM/Source/ADMSpatialization/Private/ADMSpatializationModule.h
@@ -0,0 +1,21 @@
+// Copyright Epic Games, Inc. All Rights Reserved.
+#pragma once
+
+#include "ADMSpatialization.h"
+#include "Modules/ModuleInterface.h"
+
+
+namespace UE::ADM::Spatialization
+{
+	class ADMSPATIALIZATION_API FModule : public IModuleInterface
+	{
+	public:
+		FADMSpatializationFactory& GetFactory();
+
+		virtual void StartupModule() override;
+		virtual void ShutdownModule() override;
+
+	private:
+		FADMSpatializationFactory SpatializationFactory;
+	};
+} // namespace UE::ADM::Spatialization
diff --git a/Engine/Plugins/Runtime/ADM/Source/ADMSpatialization/Private/ADMSpatializationSettings.cpp b/Engine/Plugins/Runtime/ADM/Source/ADMSpatialization/Private/ADMSpatializationSettings.cpp
new file mode 100644
--- /dev/null
+++ b/Engine/Plugins/Runtime/ADM/Source/ADMSpatialization/Private/ADMSpatializationSettings.cpp
@@ -0,0 +1,56 @@
+// Copyright Epic Games, Inc. All Rights Reserved.
+
+#include "ADMSpatializationSettings.h"
+
+#include "ADMSpatialization.h"
+#include "AudioDevice.h"
+#include "AudioDeviceManager.h"
+#include "Interfaces/IPv4/IPv4Endpoint.h"
+#include "ADMSpatializationModule.h"
+
+
+#if WITH_EDITOR
+void UADMSpatializationSettings::ADMConnect()
+{
+	using namespace UE::ADM::Spatialization;
+
+	FIPv4Address ADMAddress;
+	FIPv4Address::Parse(IPAddress, ADMAddress);
+	FIPv4Endpoint ADMEndpoint(ADMAddress, IPPort);
+
+	FModule& ADMSpatializationModule = FModuleManager::Get().LoadModuleChecked<FModule>("ADMSpatialization");
+	FADMSpatializationFactory& SpatFactory = ADMSpatializationModule.GetFactory();
+	SpatFactory.SetSendIPEndpoint(ADMEndpoint);
+
+	FAudioDeviceManager* AudioDeviceManager = FAudioDeviceManager::Get();
+	check(AudioDeviceManager);
+
+	AudioDeviceManager->IterateOverAllDevices([&SpatFactory, &ADMEndpoint](Audio::FDeviceId DeviceId, FAudioDevice* AudioDevice)
+	{
+		if (AudioDevice)
+		{
+			FAudioDevice::FAudioSpatializationInterfaceInfo Info = AudioDevice->GetCurrentSpatializationPluginInterfaceInfo();
+			if (Info.PluginName == SpatFactory.GetDisplayName())
+			{
+				if (FADMSpatialization* Spatializer = static_cast<FADMSpatialization*>(AudioDevice->GetSpatializationPluginInterface().Get()))
+				{
+					Spatializer->SetClient(FADMClient(ADMEndpoint));
+				}
+			}
+		}
+	});
+}
+
+void UADMSpatializationSettings::PostEditChangeProperty(FPropertyChangedEvent& InPropertyChangedEvent)
+{
+	const FName PropName = InPropertyChangedEvent.GetPropertyName();
+	const bool bUpdatedADMClientSettings = 
+		PropName == GET_MEMBER_NAME_CHECKED(UADMSpatializationSettings, IPAddress) ||
+		PropName == GET_MEMBER_NAME_CHECKED(UADMSpatializationSettings, IPPort);
+
+	if (bUpdatedADMClientSettings)
+	{
+		ADMConnect();
+	}
+}
+#endif // WITH_EDITOR
diff --git a/Engine/Plugins/Runtime/ADM/Source/ADMSpatialization/Private/ADMSpatializationSettings.h b/Engine/Plugins/Runtime/ADM/Source/ADMSpatialization/Private/ADMSpatializationSettings.h
new file mode 100644
--- /dev/null
+++ b/Engine/Plugins/Runtime/ADM/Source/ADMSpatialization/Private/ADMSpatializationSettings.h
@@ -0,0 +1,34 @@
+// Copyright Epic Games, Inc. All Rights Reserved.
+#pragma once
+
+#include "Engine/DeveloperSettings.h"
+#include "UObject/Object.h"
+#include "UObject/UnrealType.h"
+
+#include "ADMSpatializationSettings.generated.h"
+
+
+UCLASS(config = ADM, defaultconfig, meta = (DisplayName = "ADM Spatialization Settings"))
+class UADMSpatializationSettings : public UDeveloperSettings
+{
+	GENERATED_BODY()
+
+public:
+	// If set to a valid IP address, enables ADM Spatialization and sends source data to the
+	// given IP Address using the OSC network protocol. Can be overridden via console variable
+	// "au.OSC.ADM.SendAddress" for configuring on stand-alone client.
+	UPROPERTY(config, EditAnywhere, Category = ADM, meta = (DisplayName = "OSC IP Address"))
+	FString IPAddress;
+
+	// Default port source data to the given IP Address.
+	UPROPERTY(config, EditAnywhere, Category = ADM, meta = (DisplayName = "OSC IP Port"))
+	int32 IPPort = 8094;
+
+#if WITH_EDITOR
+	// Initializes (or reinitializes) the ADM client with the given settings.
+	UFUNCTION(CallInEditor, Category = ADM, meta = (DisplayName = "Connect ADM OSC Client"))
+	void ADMConnect();
+
+	virtual void PostEditChangeProperty(FPropertyChangedEvent& InPropertyChangedEvent) override;
+#endif // WITH_EDITOR
+};
