Add direct out channel map which maps source IDs to channel/object IDs to ADM plugin.

#rnx

diff --git a/Engine/Plugins/Runtime/ADM/Source/ADMSpatialization/Private/ADMDirectOutChannel.cpp b/Engine/Plugins/Runtime/ADM/Source/ADMSpatialization/Private/ADMDirectOutChannel.cpp
--- a/Engine/Plugins/Runtime/ADM/Source/ADMSpatialization/Private/ADMDirectOutChannel.cpp
+++ b/Engine/Plugins/Runtime/ADM/Source/ADMSpatialization/Private/ADMDirectOutChannel.cpp
@@ -62,7 +62,7 @@ namespace UE::ADM::Spatialization
 		DestructorCriticalSection.Unlock();
 	}
 
-	void FSourceDirectOut::ProcessSilience()
+	void FSourceDirectOut::ProcessSilence()
 	{
 		if (!DestructorCriticalSection.TryLock())
 		{
diff --git a/Engine/Plugins/Runtime/ADM/Source/ADMSpatialization/Private/ADMDirectOutChannel.h b/Engine/Plugins/Runtime/ADM/Source/ADMSpatialization/Private/ADMDirectOutChannel.h
--- a/Engine/Plugins/Runtime/ADM/Source/ADMSpatialization/Private/ADMDirectOutChannel.h
+++ b/Engine/Plugins/Runtime/ADM/Source/ADMSpatialization/Private/ADMDirectOutChannel.h
@@ -18,11 +18,14 @@ namespace UE::ADM::Spatialization
 		~FSourceDirectOut();
 
 		void ProcessDirectOut(const FAudioPluginSourceInputData& InInputData);
-		void ProcessSilience();
+		void ProcessSilence();
 
 		void SetIsActive(const bool bInIsActive) { bIsActive = bInIsActive;  }
 		bool GetIsActive() const { return bIsActive; }
 
+		void SetSourceId(const int32 InSourceId) { SourceId = InSourceId; }
+		int32 GetSourceId() const { return SourceId; }
+
 	private:
 		static TUniquePtr<Audio::FAlignedFloatBuffer> SilienceBuffer;
 
@@ -30,6 +33,7 @@ namespace UE::ADM::Spatialization
 
 		bool bIsActive = false;
 		Audio::IAudioMixerPlatformInterface* MixerPlatform = nullptr;
+		int32 SourceId = INDEX_NONE;
 		int32 ChannelIndex = INDEX_NONE;
 
 		void InitializeTempBuffers(const uint32 InNumSamples);
diff --git a/Engine/Plugins/Runtime/ADM/Source/ADMSpatialization/Private/ADMSpatialization.cpp b/Engine/Plugins/Runtime/ADM/Source/ADMSpatialization/Private/ADMSpatialization.cpp
--- a/Engine/Plugins/Runtime/ADM/Source/ADMSpatialization/Private/ADMSpatialization.cpp
+++ b/Engine/Plugins/Runtime/ADM/Source/ADMSpatialization/Private/ADMSpatialization.cpp
@@ -10,6 +10,7 @@
 #include "OSCMessage.h"
 #include "OSCTypes.h"
 
+#include UE_INLINE_GENERATED_CPP_BY_NAME(ADMSpatialization)
 
 void UADMEngineSubsystem::SetSendIPAddress(const FString& InIPAddress, int32 Port)
 {
@@ -41,19 +42,20 @@ namespace UE::ADM::Spatialization
 			TEXT("Default: Empty (Does not override project setting)"),
 			ECVF_Default);
 
-		int32 PositionAddressOffsetCVar = 0;
+		int32 PositionAddressOffsetCVar = -1;
 		FAutoConsoleVariableRef CVarOSCADMPositionAddressOffset(
 			TEXT("au.ADM.Spatialization.OSCPositionAddressOffset"),
-			SendEndpointCVar,
-			TEXT("Applies an index offset to all source Ids translated to OSC position source addresses.")
-			TEXT("Default: 0 (No offset)"),
+			PositionAddressOffsetCVar,
+			TEXT("Applies an index offset to all object Ids translated to OSC position source addresses.")
+			TEXT("Default: -1 (Use system default offset)"),
 			ECVF_Default);
 
 		static const FString SystemName(TEXT("ADM Spatialization"));
 
 	} // namespace SpatPrivate
 
-	FADMClient::FADMClient(const FIPv4Endpoint& InEndpoint)
+	FADMClient::FADMClient(const FIPv4Endpoint& InEndpoint, const int32 InObjectIndexOffset) :
+		ObjectIndexOffset(InObjectIndexOffset)
 	{
 		FIPv4Endpoint IPEndpoint = InEndpoint;
 		if (!SpatPrivate::SendEndpointCVar.IsEmpty())
@@ -65,13 +67,25 @@ namespace UE::ADM::Spatialization
 		ClientProxy->SetSendIPEndpoint(IPEndpoint);
 	}
 
-	FOSCAddress FADMClient::CreateConfigAddress(int32 SourceId, FString InMethod)
+	int32 FADMClient::GetObjectIndexOffset() const
+	{
+		if (SpatPrivate::PositionAddressOffsetCVar >= 0)
+		{
+			return SpatPrivate::PositionAddressOffsetCVar;
+		}
+		else
+		{
+			return ObjectIndexOffset;
+		}
+	}
+
+	FOSCAddress FADMClient::CreateConfigAddress(int32 InObjIndex, FString InMethod) const
 	{
 		TArray<FString> Containers
 		{
 			TEXT("adm"),
 			TEXT("obj"),
-			FString::FromInt(SourceId),
+			FString::FromInt(InObjIndex + GetObjectIndexOffset()),
 			TEXT("config")
 		};
 
@@ -80,13 +94,13 @@ namespace UE::ADM::Spatialization
 		return OSCAddress;
 	}
 
-	FOSCAddress FADMClient::CreatePositionAddress(int32 SourceId, FString InMethod)
+	FOSCAddress FADMClient::CreatePositionAddress(int32 InObjIndex, FString InMethod) const
 	{
 		TArray<FString> Containers
 		{
 			TEXT("adm"),
 			TEXT("obj"),
-			FString::FromInt(SourceId + SpatPrivate::PositionAddressOffsetCVar)
+			FString::FromInt(InObjIndex + GetObjectIndexOffset())
 		};
 
 		FOSCAddress OSCAddress;
@@ -94,25 +108,25 @@ namespace UE::ADM::Spatialization
 		return OSCAddress;
 	}
 
-	void FADMClient::InitSource(int32 InSourceId, bool bCartesian)
+	void FADMClient::InitObjectIndex(int32 InObjIndex, bool bCartesian)
 	{
 		using namespace UE::OSC;
 
 		if (ClientProxy.IsValid())
 		{
-			FOSCAddress OSCAddress = CreateConfigAddress(InSourceId, TEXT("cartesian"));
+			FOSCAddress OSCAddress = CreateConfigAddress(InObjIndex, TEXT("cartesian"));
 			FOSCMessage Msg(MoveTemp(OSCAddress), { FOSCData((int32)(bCartesian)) });
 			ClientProxy->SendMessage(Msg);
 		}
 	}
 
-	void FADMClient::SetPosition(int32 InSourceId, const FVector& InPosition)
+	void FADMClient::SetPosition(int32 InObjIndex, const FVector& InPosition)
 	{
 		using namespace UE::OSC;
 
 		if (ClientProxy.IsValid())
 		{
-			FOSCAddress OSCAddress = CreatePositionAddress(InSourceId, TEXT("xyz"));
+			FOSCAddress OSCAddress = CreatePositionAddress(InObjIndex, TEXT("xyz"));
 			FOSCMessage Msg(MoveTemp(OSCAddress), { FOSCData((float)InPosition.X) , FOSCData((float)InPosition.Y), FOSCData((float)InPosition.Z) });
 			ClientProxy->SendMessage(Msg);
 		}
@@ -145,7 +159,9 @@ namespace UE::ADM::Spatialization
 
 	bool FADMSpatializationFactory::IsExternalSend()
 	{
-		return false;
+		// Because this plugin sends all spatialized sources to direct outputs,
+		// indicate to the engine that it acts as an external send.
+		return true;
 	}
 
 	int32 FADMSpatializationFactory::GetMaxSupportedChannels()
@@ -159,8 +175,6 @@ namespace UE::ADM::Spatialization
 		FADMSpatializationFactory& SpatFactory = Module.GetFactory();
 		const FIPv4Endpoint Endpoint = SpatFactory.GetSendIPEndpoint();
 
-		SetClient(FADMClient(Endpoint));
-
 		int32 NumDirectOutChannels = 0;
 		Audio::FMixerDevice* AudioMixerDevice = static_cast<Audio::FMixerDevice*>(InitializationParams.AudioDevicePtr);
 		if (!ensureMsgf(AudioMixerDevice, TEXT("Unable to initialize FADMSpatialization, null audio device")))
@@ -188,11 +202,16 @@ namespace UE::ADM::Spatialization
 		{
 			DirectOuts.Emplace(Index, InitializationParams.BufferLength, MixerPlatform);
 		}
+
+		SourceIdChannelMap.Empty(InitializationParams.NumSources);
+
+		SetClient(FADMClient(Endpoint, NumBedChannels));
 	}
 
 	void FADMSpatialization::Shutdown()
 	{
 		DirectOuts.Reset();
+		SourceIdChannelMap.Reset();
 	}
 
 	void FADMSpatialization::SetClient(FADMClient&& InClient)
@@ -209,25 +228,68 @@ namespace UE::ADM::Spatialization
 	{
 		checkf((int32)SourceId < NumSources, TEXT("SourceId is expected to be less than the initialized max number of sources"));
 
-		// Use SourceId as channel index
-		DirectOuts[SourceId].SetIsActive(true);
+		if (ensure(!SourceIdChannelMap.Contains(SourceId)))
+		{
+			int32 ChannelIndex = INDEX_NONE;
+			for (int32 Index = 0; Index < DirectOuts.Num(); ++Index)
+			{
+				if (!DirectOuts[Index].GetIsActive())
+				{
+					ChannelIndex = Index;
+					break;
+				}
+			}
+
+			if (DirectOuts.IsValidIndex(ChannelIndex))
+			{
+				DirectOuts[ChannelIndex].SetIsActive(true);
+				DirectOuts[ChannelIndex].SetSourceId(SourceId);
 
-		Client.InitSource((int32)(SourceId));
+				SourceIdChannelMap.Emplace(SourceId, ChannelIndex);
+
+				// The client uses the channel index as the object index per
+				// the ADM OSC spec
+				Client.InitObjectIndex(ChannelIndex);
+			}
+			else
+			{
+				UE_LOG(LogADMSpatialization, Warning, TEXT("Failed to find available direct out channel for SourceId: %d"), SourceId);
+			}
+		}
 	}
 
 	void FADMSpatialization::OnReleaseSource(const uint32 SourceId)
 	{
 		checkf((int32)SourceId < NumSources, TEXT("SourceId is expected to be less than the initialized max number of sources"));
 
-		DirectOuts[SourceId].SetIsActive(false);
+		if (ensure(SourceIdChannelMap.Contains(SourceId)))
+		{
+			const int32 ChannelIndex = SourceIdChannelMap[SourceId];
+
+			if (ensure(DirectOuts.IsValidIndex(ChannelIndex)))
+			{
+				DirectOuts[ChannelIndex].SetIsActive(false);
+				DirectOuts[ChannelIndex].SetSourceId(INDEX_NONE);
+			}
+
+			SourceIdChannelMap.Remove(SourceId);
+		}
 	}
 
 	void FADMSpatialization::ProcessAudio(const FAudioPluginSourceInputData& InputData, FAudioPluginSourceOutputData& OutputData)
 	{
 		checkf(InputData.SourceId < NumSources, TEXT("SourceId is expected to be less than the initialized max number of sources"));
 
-		DirectOuts[InputData.SourceId].ProcessDirectOut(InputData);
-		Client.SetPosition(InputData.SourceId, InputData.SpatializationParams->EmitterPosition);
+		if (SourceIdChannelMap.Contains(InputData.SourceId))
+		{
+			const int32 ChannelIndex = SourceIdChannelMap[InputData.SourceId];
+
+			if (DirectOuts.IsValidIndex(ChannelIndex) && DirectOuts[ChannelIndex].GetIsActive())
+			{
+				DirectOuts[ChannelIndex].ProcessDirectOut(InputData);
+				Client.SetPosition(ChannelIndex, InputData.SpatializationParams->EmitterPosition);
+			}
+		}
 	}
 
 	void FADMSpatialization::OnAllSourcesProcessed()
@@ -237,7 +299,7 @@ namespace UE::ADM::Spatialization
 			if (!DirectOut.GetIsActive())
 			{
 				// Send silience out non-active outputs to retain time alignment with other outputs
-				DirectOut.ProcessSilience();
+				DirectOut.ProcessSilence();
 			}
 		}
 	}
diff --git a/Engine/Plugins/Runtime/ADM/Source/ADMSpatialization/Private/ADMSpatialization.h b/Engine/Plugins/Runtime/ADM/Source/ADMSpatialization/Private/ADMSpatialization.h
--- a/Engine/Plugins/Runtime/ADM/Source/ADMSpatialization/Private/ADMSpatialization.h
+++ b/Engine/Plugins/Runtime/ADM/Source/ADMSpatialization/Private/ADMSpatialization.h
@@ -82,18 +82,21 @@ namespace UE::ADM::Spatialization
 	{
 	public:
 		FADMClient() = default;
-		FADMClient(const FIPv4Endpoint& InEndpoint);
+		FADMClient(const FIPv4Endpoint& InEndpoint, const int32 InObjectIndexOffset = 0);
 
-		static FOSCAddress CreateConfigAddress(int32 SourceId, FString InMethod);
-		static FOSCAddress CreatePositionAddress(int32 SourceId, FString InMethod);
-
-		void InitSource(int32 InSourceId, bool bCartesian = true);
+		void InitObjectIndex(int32 InObjIndex, bool bCartesian = true);
 		bool IsSet() const;
 
-		void SetPosition(int32 InSourceId, const FVector& InPosition);
+		void SetPosition(int32 InObjIndex, const FVector& InPosition);
 
 	private:
 		TUniquePtr<UE::OSC::IClientProxy> ClientProxy;
+		int32 ObjectIndexOffset = 0;
+
+		int32 GetObjectIndexOffset() const;
+
+		FOSCAddress CreateConfigAddress(int32 InObjIndex, FString InMethod) const;
+		FOSCAddress CreatePositionAddress(int32 InObjIndex, FString InMethod) const;
 	};
 
 
@@ -119,6 +122,8 @@ namespace UE::ADM::Spatialization
 		float SampleRate = 0.f;
 
 		TArray<FSourceDirectOut> DirectOuts;
+		/** Maps Source ID to Channel index */
+		TMap<int32, int32> SourceIdChannelMap;
 
 		FADMClient Client;
 	};
diff --git a/Engine/Source/Runtime/Windows/AudioMixerWasapi/Private/AudioMixerWasapi.cpp b/Engine/Source/Runtime/Windows/AudioMixerWasapi/Private/AudioMixerWasapi.cpp
--- a/Engine/Source/Runtime/Windows/AudioMixerWasapi/Private/AudioMixerWasapi.cpp
+++ b/Engine/Source/Runtime/Windows/AudioMixerWasapi/Private/AudioMixerWasapi.cpp
@@ -217,7 +217,8 @@ namespace Audio
 
 		bIsInitialized = true;
 
-		UE_LOG(LogAudioMixer, Display, TEXT("FAudioMixerWasapi initialized SampeRate: %d"), OpenStreamParams.SampleRate);
+		UE_LOG(LogAudioMixer, Display, TEXT("FAudioMixerWasapi initialized: SampleRate=%d NumChannels=%d NumDirectOutChannels=%d AggregateMode=%d"), 
+			OpenStreamParams.SampleRate, AudioStreamInfo.DeviceInfo.NumChannels, AudioStreamInfo.DeviceInfo.NumDirectOutChannels, bUseWasapiAggregateDevice);
 
 		return true;
 	}
