Fix incorrect UE to ADM-OSC coordinate translation.

#rnx

diff --git a/Engine/Plugins/Runtime/ADM/Source/ADMSpatialization/Private/ADMSpatialization.cpp b/Engine/Plugins/Runtime/ADM/Source/ADMSpatialization/Private/ADMSpatialization.cpp
--- a/Engine/Plugins/Runtime/ADM/Source/ADMSpatialization/Private/ADMSpatialization.cpp
+++ b/Engine/Plugins/Runtime/ADM/Source/ADMSpatialization/Private/ADMSpatialization.cpp
@@ -6,6 +6,7 @@
 #include "ADMSpatializationLog.h"
 #include "ADMSpatializationModule.h"
 #include "AudioMixerDevice.h"
+#include "Math/Transform.h"
 #include "OSCClient.h"
 #include "OSCMessage.h"
 #include "OSCTypes.h"
@@ -117,6 +118,8 @@ namespace UE::ADM::Spatialization
 			FOSCAddress OSCAddress = CreateConfigAddress(InObjIndex, TEXT("cartesian"));
 			FOSCMessage Msg(MoveTemp(OSCAddress), { FOSCData((int32)(bCartesian)) });
 			ClientProxy->SendMessage(Msg);
+
+			UE_LOG(LogADMSpatialization, Verbose, TEXT("InitObjectIndex: %d"), InObjIndex + GetObjectIndexOffset());
 		}
 	}
 
@@ -126,8 +129,11 @@ namespace UE::ADM::Spatialization
 
 		if (ClientProxy.IsValid())
 		{
+			FVector ADMPosition = UnrealToADMCoordinates(InPosition);
+
 			FOSCAddress OSCAddress = CreatePositionAddress(InObjIndex, TEXT("xyz"));
-			FOSCMessage Msg(MoveTemp(OSCAddress), { FOSCData((float)InPosition.X) , FOSCData((float)InPosition.Y), FOSCData((float)InPosition.Z) });
+			FOSCMessage Msg(MoveTemp(OSCAddress), { FOSCData((float)ADMPosition.X) , FOSCData((float)ADMPosition.Y), FOSCData((float)ADMPosition.Z) });
+
 			ClientProxy->SendMessage(Msg);
 		}
 	}
@@ -250,6 +256,8 @@ namespace UE::ADM::Spatialization
 				// The client uses the channel index as the object index per
 				// the ADM OSC spec
 				Client.InitObjectIndex(ChannelIndex);
+
+				UE_LOG(LogADMSpatialization, Verbose, TEXT("OnInitSource: %d [%d]"), SourceId, ChannelIndex);
 			}
 			else
 			{
@@ -287,7 +295,11 @@ namespace UE::ADM::Spatialization
 			if (DirectOuts.IsValidIndex(ChannelIndex) && DirectOuts[ChannelIndex].GetIsActive())
 			{
 				DirectOuts[ChannelIndex].ProcessDirectOut(InputData);
-				Client.SetPosition(ChannelIndex, InputData.SpatializationParams->EmitterPosition);
+
+				FTransform ListenerTransform = FTransform(InputData.SpatializationParams->ListenerPosition);
+				FVector ListenerRelEmitterPos = ListenerTransform.Inverse().TransformPosition(InputData.SpatializationParams->EmitterWorldPosition);
+
+				Client.SetPosition(ChannelIndex, ListenerRelEmitterPos.GetSafeNormal());
 			}
 		}
 	}
diff --git a/Engine/Plugins/Runtime/ADM/Source/ADMSpatialization/Private/ADMSpatialization.h b/Engine/Plugins/Runtime/ADM/Source/ADMSpatialization/Private/ADMSpatialization.h
--- a/Engine/Plugins/Runtime/ADM/Source/ADMSpatialization/Private/ADMSpatialization.h
+++ b/Engine/Plugins/Runtime/ADM/Source/ADMSpatialization/Private/ADMSpatialization.h
@@ -97,6 +97,21 @@ namespace UE::ADM::Spatialization
 
 		FOSCAddress CreateConfigAddress(int32 InObjIndex, FString InMethod) const;
 		FOSCAddress CreatePositionAddress(int32 InObjIndex, FString InMethod) const;
+
+		FORCEINLINE FVector UnrealToADMCoordinates(const FVector& InPosition)
+		{
+			/*       
+				UNREAL                    ADM-OSC
+				 Z                          Z
+				 |    X                     |    Y
+				 |   /                      |   /
+				 |  /                       |  /
+				 | /                        | /
+				 |/_______________Y         |/_______________X
+			*/
+
+			return { InPosition.Y, InPosition.X, InPosition.Z };
+		}
 	};
 
 
diff --git a/Engine/Plugins/Runtime/ADM/Source/ADMSpatialization/Private/ADMSpatializationSettings.h b/Engine/Plugins/Runtime/ADM/Source/ADMSpatialization/Private/ADMSpatializationSettings.h
--- a/Engine/Plugins/Runtime/ADM/Source/ADMSpatialization/Private/ADMSpatializationSettings.h
+++ b/Engine/Plugins/Runtime/ADM/Source/ADMSpatialization/Private/ADMSpatializationSettings.h
@@ -20,9 +20,10 @@ public:
 	UPROPERTY(config, EditAnywhere, Category = ADM, meta = (DisplayName = "OSC IP Address"))
 	FString IPAddress;
 
-	// Default port source data to the given IP Address.
+	// The IP port used in conjunction with the given IP Address. Default value of 4001
+	// per the ADM-OSC spec.
 	UPROPERTY(config, EditAnywhere, Category = ADM, meta = (DisplayName = "OSC IP Port"))
-	int32 IPPort = 8094;
+	int32 IPPort = 4001;
 
 #if WITH_EDITOR
 	// Initializes (or reinitializes) the ADM client with the given settings.
